Allow non-editor runs on the Build Farm to capture their Horde environment when using StudioTelemetry
Removed the dependency on Horde library completely from Studio Telemetry and re-implemented as direct commandline/environment capture calls.

diff --git a/Engine/Source/Runtime/StudioTelemetry/Private/StudioTelemetry.cpp b/Engine/Source/Runtime/StudioTelemetry/Private/StudioTelemetry.cpp
--- a/Engine/Source/Runtime/StudioTelemetry/Private/StudioTelemetry.cpp
+++ b/Engine/Source/Runtime/StudioTelemetry/Private/StudioTelemetry.cpp
@@ -2,11 +2,6 @@
 
 #include "StudioTelemetry.h"
 #include "StudioTelemetryLog.h"
-
-#if WITH_EDITOR
-#include "Horde.h"
-#endif
-
 #include "Analytics.h"
 #include "AnalyticsProviderBroadcast.h"
 #include "AnalyticsTracer.h"
@@ -28,6 +23,16 @@ DEFINE_LOG_CATEGORY(LogStudioTelemetry);
 
 IMPLEMENT_MODULE(FStudioTelemetry, StudioTelemetry)
 
+FString GetEnvironmentValriable(const FString& CommandLineVar, const FString& EnvironmentVar)
+{
+	FString Result;
+	if (false == FParse::Value(FCommandLine::Get(), *CommandLineVar, Result))
+	{
+		Result = FPlatformMisc::GetEnvironmentVariable(*EnvironmentVar);
+	}
+	return Result;
+}
+
 FStudioTelemetry& FStudioTelemetry::Get()
 {
 	static FStudioTelemetry StudioTelemetryInstance;
@@ -151,21 +156,19 @@ void FStudioTelemetry::StartSession()
 			DefaultEventAttributes.Emplace(TEXT("OS_ID"), FPlatformMisc::GetOperatingSystemId());
 		}
 
-#if WITH_EDITOR
-		if (!FHorde::GetJobId().IsEmpty())
+		const FString HordeJobId = GetEnvironmentValriable(TEXT("HordeJobId="), TEXT("UE_HORDE_JOBID")) ;
+
+		if (!HordeJobId.IsEmpty())
 		{
-			// Only send Horde data if applicable
-			DefaultEventAttributes.Emplace(TEXT("Horde_ServerURL"), FHorde::GetServerURL());
-			DefaultEventAttributes.Emplace(TEXT("Horde_TemplateID"), FHorde::GetTemplateId());
-			DefaultEventAttributes.Emplace(TEXT("Horde_TemplateName"), FHorde::GetTemplateName());
-			DefaultEventAttributes.Emplace(TEXT("Horde_JobURL"), FHorde::GetJobURL());
-			DefaultEventAttributes.Emplace(TEXT("Horde_JobID"), FHorde::GetJobId());
-			DefaultEventAttributes.Emplace(TEXT("Horde_StepName"), FHorde::GetStepName());
-			DefaultEventAttributes.Emplace(TEXT("Horde_StepID"), FHorde::GetStepId());
-			DefaultEventAttributes.Emplace(TEXT("Horde_StepURL"), FHorde::GetStepURL());
-			DefaultEventAttributes.Emplace(TEXT("Horde_BatchID"), FHorde::GetBatchId());
+			// Send Horde environment if applicable
+			DefaultEventAttributes.Emplace(TEXT("Horde_JobID"), HordeJobId);
+			DefaultEventAttributes.Emplace(TEXT("Horde_ServerURL"), GetEnvironmentValriable(TEXT("HordeServerUrl="), TEXT("UE_HORDE_URL")));
+			DefaultEventAttributes.Emplace(TEXT("Horde_TemplateID"), GetEnvironmentValriable(TEXT("HordeTemplateId="), TEXT("UE_HORDE_TEMPLATEID")));
+			DefaultEventAttributes.Emplace(TEXT("Horde_TemplateName"), GetEnvironmentValriable(TEXT("HordeTemplateName="), TEXT("UE_HORDE_TEMPLATENAME")));	
+			DefaultEventAttributes.Emplace(TEXT("Horde_StepName"), GetEnvironmentValriable(TEXT("HordeStepName="), TEXT("UE_HORDE_STEPNAME")));
+			DefaultEventAttributes.Emplace(TEXT("Horde_StepID"), GetEnvironmentValriable(TEXT("HordeStepId="), TEXT("UE_HORDE_STEPID")));
+			DefaultEventAttributes.Emplace(TEXT("Horde_BatchID"), GetEnvironmentValriable(TEXT("HordeBatchId="), TEXT("UE_HORDE_BATCHID")));
 		}
-#endif
 
 		// Set up the analytics provider
 		AnalyticsProvider->SetUserID(UserID);
diff --git a/Engine/Source/Runtime/StudioTelemetry/StudioTelemetry.build.cs b/Engine/Source/Runtime/StudioTelemetry/StudioTelemetry.build.cs
--- a/Engine/Source/Runtime/StudioTelemetry/StudioTelemetry.build.cs
+++ b/Engine/Source/Runtime/StudioTelemetry/StudioTelemetry.build.cs
@@ -33,8 +33,7 @@ public class StudioTelemetry : ModuleRules
 			PrivateDependencyModuleNames.AddRange(
 				new string[] {
 					"AnalyticsLog",
-					"AnalyticsHorde",
-					"Horde"
+					"AnalyticsHorde"
 				}
 			);
 		}
