Made a number of StudioTelemetry's default attributes optional to reduce the size of the payloads
In addition to  User Data, Hardware Data and OS Data are now optional
Use -ST_SendUserData=true, -ST_SendOSData=true, -ST_SendHardwareData=true from the commandline or
set SendUserData=true, SendOSData=true, SendHardwareData=true in the [StudioTelemetry.Config] section of the Engine.ini file

diff --git a/Engine/Config/BaseEngine.ini b/Engine/Config/BaseEngine.ini
--- a/Engine/Config/BaseEngine.ini
+++ b/Engine/Config/BaseEngine.ini
@@ -3784,7 +3784,7 @@ SendOSData=true
 ; UsageType=EditorAndClient
 ; APIKeyET=HordeAnalytics.Dev
 ; APIServerET="https://myownhordeserver.com/" ; Point Server URL to your own Horde server
-; APIEndpointET="api/v1/telemetry"
+; APIEndpointET="api/v1/telemetry/engine"
 
 [/Script/Engine.WorldPartitionSettings]
 EditorHashDefaultClass=/Script/CoreUObject.Class'/Script/Engine.WorldPartitionEditorSpatialHash'
diff --git a/Engine/Source/Runtime/StudioTelemetry/Private/StudioTelemetry.cpp b/Engine/Source/Runtime/StudioTelemetry/Private/StudioTelemetry.cpp
--- a/Engine/Source/Runtime/StudioTelemetry/Private/StudioTelemetry.cpp
+++ b/Engine/Source/Runtime/StudioTelemetry/Private/StudioTelemetry.cpp
@@ -102,7 +102,7 @@ void FStudioTelemetry::StartSession()
 		FString SessionLabel;
 		FParse::Value(FCommandLine::Get(), TEXT("SessionLabel="), SessionLabel);
 
-		// Set the default event attributes, these will always be sent to telemetry for every event
+		// Set the default event attributes, these will be appended to every event
 		DefaultEventAttributes.Emplace(TEXT("ProjectName"), ProjectName);
 		DefaultEventAttributes.Emplace(TEXT("ProjectID"), ProjectID);
 
@@ -111,18 +111,13 @@ void FStudioTelemetry::StartSession()
 		DefaultEventAttributes.Emplace(TEXT("Session_StartUTC"), FDateTime::UtcNow().ToUnixTimestampDecimal());
 
 		DefaultEventAttributes.Emplace(TEXT("Build_Configuration"), LexToString(FApp::GetBuildConfiguration()));
-		DefaultEventAttributes.Emplace(TEXT("Build_IsInternalBuild"), FEngineBuildSettings::IsInternalBuild());
-		DefaultEventAttributes.Emplace(TEXT("Build_IsPerforceBuild"), FEngineBuildSettings::IsPerforceBuild());
-		DefaultEventAttributes.Emplace(TEXT("Build_IsPromotedBuild"), FApp::GetEngineIsPromotedBuild() == 0 ? false : true);
 		DefaultEventAttributes.Emplace(TEXT("Build_BranchName"), FApp::GetBranchName().ToLower());
 		DefaultEventAttributes.Emplace(TEXT("Build_Changelist"), BuildSettings::GetCurrentChangelist());
 
 		DefaultEventAttributes.Emplace(TEXT("Config_IsEditor"), GIsEditor);
-		DefaultEventAttributes.Emplace(TEXT("Config_IsUnattended"), FApp::IsUnattended());
 		DefaultEventAttributes.Emplace(TEXT("Config_IsBuildMachine"), GIsBuildMachine);
 		DefaultEventAttributes.Emplace(TEXT("Config_IsRunningCommandlet"), IsRunningCommandlet());
-		DefaultEventAttributes.Emplace(TEXT("Config_IsDebuggerPresent"), FPlatformMisc::IsDebuggerPresent());
-
+		
 		// Only send user data if requested
 		if (Config.bSendUserData == true)
 		{
@@ -158,16 +153,14 @@ void FStudioTelemetry::StartSession()
 
 		const FString HordeJobId = GetEnvironmentValriable(TEXT("HordeJobId="), TEXT("UE_HORDE_JOBID")) ;
 
+		// Send Horde environment if available
 		if (!HordeJobId.IsEmpty())
 		{
-			// Send Horde environment if applicable
 			DefaultEventAttributes.Emplace(TEXT("Horde_JobID"), HordeJobId);
+			DefaultEventAttributes.Emplace(TEXT("Horde_StepID"), GetEnvironmentValriable(TEXT("HordeStepId="), TEXT("UE_HORDE_STEPID")));
+			DefaultEventAttributes.Emplace(TEXT("Horde_StepName"), GetEnvironmentValriable(TEXT("HordeStepName="), TEXT("UE_HORDE_STEPNAME")));
 			DefaultEventAttributes.Emplace(TEXT("Horde_ServerURL"), GetEnvironmentValriable(TEXT("HordeServerUrl="), TEXT("UE_HORDE_URL")));
-			DefaultEventAttributes.Emplace(TEXT("Horde_TemplateID"), GetEnvironmentValriable(TEXT("HordeTemplateId="), TEXT("UE_HORDE_TEMPLATEID")));
 			DefaultEventAttributes.Emplace(TEXT("Horde_TemplateName"), GetEnvironmentValriable(TEXT("HordeTemplateName="), TEXT("UE_HORDE_TEMPLATENAME")));	
-			DefaultEventAttributes.Emplace(TEXT("Horde_StepName"), GetEnvironmentValriable(TEXT("HordeStepName="), TEXT("UE_HORDE_STEPNAME")));
-			DefaultEventAttributes.Emplace(TEXT("Horde_StepID"), GetEnvironmentValriable(TEXT("HordeStepId="), TEXT("UE_HORDE_STEPID")));
-			DefaultEventAttributes.Emplace(TEXT("Horde_BatchID"), GetEnvironmentValriable(TEXT("HordeBatchId="), TEXT("UE_HORDE_BATCHID")));
 		}
 
 		// Set up the analytics provider
diff --git a/Engine/Source/Runtime/StudioTelemetry/Public/StudioTelemetry.h b/Engine/Source/Runtime/StudioTelemetry/Public/StudioTelemetry.h
--- a/Engine/Source/Runtime/StudioTelemetry/Public/StudioTelemetry.h
+++ b/Engine/Source/Runtime/StudioTelemetry/Public/StudioTelemetry.h
@@ -127,8 +127,8 @@ private:
 	{
 		bool bSendTelemetry = true; // Only send telemetry data if we have been requested to
 		bool bSendUserData = false;  // Never send user data unless specifically asked to
-		bool bSendHardwareData = true; // Always send hardware data unless specifically asked not to
-		bool bSendOSData = true; // Always send operating system data unless specifically asked not to
+		bool bSendHardwareData = false; // Never send hardware data unless specifically asked to4
+		bool bSendOSData = false; // Never send operating system data unless specifically asked to
 	};
 	
 	FCriticalSection						CriticalSection;
diff --git a/Engine/Source/Runtime/StudioTelemetry/StudioTelemetry.build.cs b/Engine/Source/Runtime/StudioTelemetry/StudioTelemetry.build.cs
--- a/Engine/Source/Runtime/StudioTelemetry/StudioTelemetry.build.cs
+++ b/Engine/Source/Runtime/StudioTelemetry/StudioTelemetry.build.cs
@@ -9,33 +9,20 @@ public class StudioTelemetry : ModuleRules
 	{
 		PCHUsage = ModuleRules.PCHUsageMode.UseExplicitOrSharedPCHs;
 		
-		PublicDependencyModuleNames.AddRange(
-			new string[] {
-				"Core",
-				"HTTP"
-			}
-		);
-
 		PrivateDependencyModuleNames.AddRange(
 			new string[] {
+				"Core",
+				"HTTP",
 				"CoreUObject",
                 "EngineSettings",
                 "BuildSettings",
                 "Analytics",
                 "AnalyticsET",
+				"AnalyticsLog",
+				"AnalyticsHorde",
 				"TelemetryUtils",
 				"RHI",
 			}
 		);
-
-		if (Target.bBuildEditor)
-		{
-			PrivateDependencyModuleNames.AddRange(
-				new string[] {
-					"AnalyticsLog",
-					"AnalyticsHorde"
-				}
-			);
-		}
 	}
 }
