Always send the Hardware_Platform name in the StudioTelemetry event

diff --git a/Engine/Source/Runtime/StudioTelemetry/Private/StudioTelemetry.cpp b/Engine/Source/Runtime/StudioTelemetry/Private/StudioTelemetry.cpp
--- a/Engine/Source/Runtime/StudioTelemetry/Private/StudioTelemetry.cpp
+++ b/Engine/Source/Runtime/StudioTelemetry/Private/StudioTelemetry.cpp
@@ -99,17 +99,19 @@ void FStudioTelemetry::StartSession()
 
 		FGuid SessionID = FApp::GetInstanceId();
 
-		FString SessionLabel;
-		FParse::Value(FCommandLine::Get(), TEXT("SessionLabel="), SessionLabel);
-
 		// Set the default event attributes, these will be appended to every event
 		DefaultEventAttributes.Emplace(TEXT("ProjectName"), ProjectName);
 		DefaultEventAttributes.Emplace(TEXT("ProjectID"), ProjectID);
 
 		DefaultEventAttributes.Emplace(TEXT("Session_ID"), SessionID.ToString(EGuidFormats::DigitsWithHyphensInBraces));
-		DefaultEventAttributes.Emplace(TEXT("Session_Label"), SessionLabel);
 		DefaultEventAttributes.Emplace(TEXT("Session_StartUTC"), FDateTime::UtcNow().ToUnixTimestampDecimal());
 
+		FString SessionLabel;
+		if (FParse::Value(FCommandLine::Get(), TEXT("SessionLabel="), SessionLabel))
+		{
+			DefaultEventAttributes.Emplace(TEXT("Session_Label"), SessionLabel);
+		}
+
 		DefaultEventAttributes.Emplace(TEXT("Build_Configuration"), LexToString(FApp::GetBuildConfiguration()));
 		DefaultEventAttributes.Emplace(TEXT("Build_BranchName"), FApp::GetBranchName().ToLower());
 		DefaultEventAttributes.Emplace(TEXT("Build_Changelist"), BuildSettings::GetCurrentChangelist());
@@ -117,7 +119,7 @@ void FStudioTelemetry::StartSession()
 		DefaultEventAttributes.Emplace(TEXT("Config_IsEditor"), GIsEditor);
 		DefaultEventAttributes.Emplace(TEXT("Config_IsBuildMachine"), GIsBuildMachine);
 		DefaultEventAttributes.Emplace(TEXT("Config_IsRunningCommandlet"), IsRunningCommandlet());
-		
+
 		// Only send user data if requested
 		if (Config.bSendUserData == true)
 		{
@@ -125,10 +127,12 @@ void FStudioTelemetry::StartSession()
 			DefaultEventAttributes.Emplace(TEXT("Application_Commandline"), FCommandLine::Get());
 		}
 
-		// Only send hardware data if requested
+		// ALways send the platform
+		DefaultEventAttributes.Emplace(TEXT("Hardware_Platform"), FString(FPlatformProperties::IniPlatformName()));
+
+		// Only send detailed hardware data if requested
 		if (Config.bSendHardwareData == true)
 		{
-			DefaultEventAttributes.Emplace(TEXT("Hardware_Platform"), FString(FPlatformProperties::IniPlatformName()));
 			DefaultEventAttributes.Emplace(TEXT("Hardware_GPU"), GRHIAdapterName);
 			DefaultEventAttributes.Emplace(TEXT("Hardware_CPU"), FPlatformMisc::GetCPUBrand());
 			DefaultEventAttributes.Emplace(TEXT("Hardware_CPU_Cores_Physical"), FPlatformMisc::NumberOfCores());
@@ -180,6 +184,10 @@ void FStudioTelemetry::StartSession()
 		// Bind the pre-exit callback
 		FCoreDelegates::OnEnginePreExit.AddRaw(&FStudioTelemetry::Get(), &FStudioTelemetry::EndSession);
 
+		TArray<FAnalyticsEventAttribute> Attributes;
+		Attributes.Emplace(TEXT("SchemaVersion"), 1);
+		AnalyticsProvider->RecordEvent(TEXT("StudioTelemetry.SessionStart"), Attributes);
+
 		OnStartSession.Broadcast();
 
 		UE_LOG(LogStudioTelemetry, Log, TEXT("Started StudioTelemetry Session"));
@@ -194,6 +202,10 @@ void FStudioTelemetry::EndSession()
 
 		OnEndSession.Broadcast();
 
+		TArray<FAnalyticsEventAttribute> Attributes;
+		Attributes.Emplace(TEXT("SchemaVersion"), 1);
+		AnalyticsProvider->RecordEvent(TEXT("StudioTelemetry.SessionEnd"), Attributes);
+
 		// End session for the tracer and the provider
 		if (AnalyticsTracer.IsValid())
 		{
