Added a commandline option to StudioTelemetry to specify a single provider in the AnalyticsBroadcastProvider.
Usage is ST_Provider=SomeProviderName.
ST_Provider=None ( or some other unregistered ProviderName ) can be used to turn off telemetry upload completely
Changed the CSVToTelemetry uplaods in the Build configuration to run as 3 separate passes each targetting a different backend.
This is a temporary workaround to avoid long riunning jobs due to timeouts only seen when issuing telemetry to mulitple backends concurently.

diff --git a/Engine/Source/Runtime/StudioTelemetry/Private/AnalyticsProviderBroadcast.cpp b/Engine/Source/Runtime/StudioTelemetry/Private/AnalyticsProviderBroadcast.cpp
--- a/Engine/Source/Runtime/StudioTelemetry/Private/AnalyticsProviderBroadcast.cpp
+++ b/Engine/Source/Runtime/StudioTelemetry/Private/AnalyticsProviderBroadcast.cpp
@@ -5,6 +5,7 @@
 #include "Interfaces/IAnalyticsProviderModule.h"
 #include "AnalyticsProviderConfigurationDelegate.h"
 #include "Misc/ConfigCacheIni.h"
+#include "Misc/CommandLine.h"
 #include "HttpModule.h"
 #include "HttpManager.h"
 
@@ -33,6 +34,13 @@ FAnalyticsProviderBroadcast::FAnalyticsProviderBroadcast()
 	const FString TelemetryProviderSection(TEXT("StudioTelemetry.Provider"));
 
 	TArray<FString> SectionNames;
+
+	FString ProviderName;
+
+	if (FParse::Value(FCommandLine::Get(), TEXT("ST_Provider="), ProviderName))
+	{
+		UE_LOG(LogAnalytics, Display, TEXT("Selected Telemetry Provider %s"), *ProviderName);
+	}
 	
 	if (GConfig->GetSectionNames(GEngineIni, SectionNames))
 	{
@@ -111,6 +119,12 @@ FAnalyticsProviderBroadcast::FAnalyticsProviderBroadcast()
 						continue;
 					}
 
+					if (ProviderName.Len() && Name != ProviderName)
+					{
+						// Skip over this provider because we don't want to use it
+						continue;
+					}
+
 					// Try to create the analytics provider
 					AnalyticsProvider = FAnalytics::Get().CreateAnalyticsProvider(FName(ProviderModuleName), FAnalyticsProviderConfigurationDelegate::CreateStatic(&GetAnalyticsProviderConfiguration));
 	
