Added a warning to indicate when a named Analytics Provider cannot be found.

diff --git a/Engine/Source/Runtime/StudioTelemetry/Private/AnalyticsProviderBroadcast.cpp b/Engine/Source/Runtime/StudioTelemetry/Private/AnalyticsProviderBroadcast.cpp
--- a/Engine/Source/Runtime/StudioTelemetry/Private/AnalyticsProviderBroadcast.cpp
+++ b/Engine/Source/Runtime/StudioTelemetry/Private/AnalyticsProviderBroadcast.cpp
@@ -145,6 +145,12 @@ FAnalyticsProviderBroadcast::FAnalyticsProviderBroadcast()
 			}
 		}
 	}
+
+	if (ProviderName.Len() && Providers.IsEmpty())
+	{
+		// We were looking to use a named provider which did not exist so raise a warning
+		UE_LOG(LogAnalytics, Warning, TEXT("Unable to find a named analytics provider %s"), *ProviderName);
+	}
 }
 
 bool FAnalyticsProviderBroadcast::SetSessionID(const FString& InSessionID)
