Removed duplicate SchemaVersion that appears for CSVToTelemetry events

diff --git a/Engine/Source/Runtime/StudioTelemetry/Private/StudioTelemetry.cpp b/Engine/Source/Runtime/StudioTelemetry/Private/StudioTelemetry.cpp
--- a/Engine/Source/Runtime/StudioTelemetry/Private/StudioTelemetry.cpp
+++ b/Engine/Source/Runtime/StudioTelemetry/Private/StudioTelemetry.cpp
@@ -186,7 +186,7 @@ void FStudioTelemetry::StartSession()
 
 		TArray<FAnalyticsEventAttribute> Attributes;
 		Attributes.Emplace(TEXT("SchemaVersion"), 1);
-		AnalyticsProvider->RecordEvent(TEXT("StudioTelemetry.SessionStart"), Attributes);
+		AnalyticsProvider->RecordEvent(TEXT("StudioTelemetry.SessionStart"), TArray<FAnalyticsEventAttribute>());
 
 		OnStartSession.Broadcast();
 
@@ -202,9 +202,7 @@ void FStudioTelemetry::EndSession()
 
 		OnEndSession.Broadcast();
 
-		TArray<FAnalyticsEventAttribute> Attributes;
-		Attributes.Emplace(TEXT("SchemaVersion"), 1);
-		AnalyticsProvider->RecordEvent(TEXT("StudioTelemetry.SessionEnd"), Attributes);
+		AnalyticsProvider->RecordEvent(TEXT("StudioTelemetry.SessionEnd"), TArray<FAnalyticsEventAttribute>());
 
 		// End session for the tracer and the provider
 		if (AnalyticsTracer.IsValid())
