Removed duplicate attribute check from IAnalyticsProvider to avoid performnce concerns
Re-introdcued this back into StudioTelemetry's AnalyticsProviderBroadcast implmentation and added back the check() for duplicates.
Deleted empty IAnalyticsProvider source

diff --git a/Engine/Source/Runtime/Analytics/Analytics/Private/Interfaces/IAnalyticsProvider.cpp b/Engine/Source/Runtime/Analytics/Analytics/Private/Interfaces/IAnalyticsProvider.cpp
deleted file mode 100644
diff --git a/Engine/Source/Runtime/Analytics/Analytics/Public/Interfaces/IAnalyticsProvider.h b/Engine/Source/Runtime/Analytics/Analytics/Public/Interfaces/IAnalyticsProvider.h
--- a/Engine/Source/Runtime/Analytics/Analytics/Public/Interfaces/IAnalyticsProvider.h
+++ b/Engine/Source/Runtime/Analytics/Analytics/Public/Interfaces/IAnalyticsProvider.h
@@ -447,9 +447,6 @@ public:
 	virtual ~IAnalyticsProvider() {}
 
 protected:
-	/**
-	 * Checks for duplicate key names in the same attribute array. Only implemented for debug and development builds.
-	 */
-	void ANALYTICS_API CheckForDuplicateAttributes(const FString& EventName, const TArray<FAnalyticsEventAttribute>& Attributes);
+
 };
 
diff --git a/Engine/Source/Runtime/Analytics/AnalyticsET/Private/IAnalyticsProviderET.cpp b/Engine/Source/Runtime/Analytics/AnalyticsET/Private/IAnalyticsProviderET.cpp
--- a/Engine/Source/Runtime/Analytics/AnalyticsET/Private/IAnalyticsProviderET.cpp
+++ b/Engine/Source/Runtime/Analytics/AnalyticsET/Private/IAnalyticsProviderET.cpp
@@ -619,9 +619,6 @@ void FAnalyticsProviderET::RecordEvent(FString&& EventName, const TArray<FAnalyt
 {
 	LLM_SCOPE_BYNAME(TEXT("Analytics"));
 
-	// Check for duplicate attributes against themselves
-	CheckForDuplicateAttributes(EventName, Attributes);
-
 	// let higher level code filter the decision of whether to send the event
 	if (ShouldRecordEvent(EventName))
 	{
diff --git a/Engine/Source/Runtime/StudioTelemetry/Private/AnalyticsProviderBroadcast.cpp b/Engine/Source/Runtime/StudioTelemetry/Private/AnalyticsProviderBroadcast.cpp
--- a/Engine/Source/Runtime/StudioTelemetry/Private/AnalyticsProviderBroadcast.cpp
+++ b/Engine/Source/Runtime/StudioTelemetry/Private/AnalyticsProviderBroadcast.cpp
@@ -259,8 +259,34 @@ void FAnalyticsProviderBroadcast::SetRecordEventCallback(OnRecordEvent Callback)
 	RecordEventCallback = Callback;
 }
 
+static void CheckForDuplicateAttributes(const FString& EventName, const TArray<FAnalyticsEventAttribute>& Attributes)
+{
+#if UE_BUILD_DEBUG || UE_BUILD_DEVELOPMENT
+	// Expose events that have duplicate attribute names. Not only does this provide ambiguous event data, in general this is not handled by all the analytics backends in any consistent manner.
+	for (int32 indexA = 0; indexA < Attributes.Num(); ++indexA)
+	{
+		const FString& AttributeName = Attributes[indexA].GetName();
+
+		for (int32 indexB = indexA + 1; indexB < Attributes.Num(); ++indexB)
+		{
+			const bool IsDuplicateAttribute = AttributeName != Attributes[indexB].GetName();
+
+			if (IsDuplicateAttribute)
+			{	
+				UE_LOG(LogAnalytics, Warning, TEXT("Duplicate Attributes %s Found For Event %s"), *AttributeName, *EventName);
+			}
+
+			check(!IsDuplicateAttribute);
+		}
+	}
+#endif
+}
+
 void FAnalyticsProviderBroadcast::RecordEvent(const FString& EventName, const TArray<FAnalyticsEventAttribute>& Attributes)
 {
+	// Check for duplicate attributes against themselves
+	CheckForDuplicateAttributes(EventName, Attributes);
+
 	for (TProviders::TConstIterator it(Providers); it; ++it)
 	{
 		(*it).Value->RecordEvent(EventName, Attributes);
