Fix for duplicate attribute test in StudioTelemetry plugin.
The duplicate compare test was the worng way round.

diff --git a/Engine/Source/Runtime/StudioTelemetry/Private/AnalyticsProviderBroadcast.cpp b/Engine/Source/Runtime/StudioTelemetry/Private/AnalyticsProviderBroadcast.cpp
--- a/Engine/Source/Runtime/StudioTelemetry/Private/AnalyticsProviderBroadcast.cpp
+++ b/Engine/Source/Runtime/StudioTelemetry/Private/AnalyticsProviderBroadcast.cpp
@@ -269,7 +269,7 @@ static void CheckForDuplicateAttributes(const FString& EventName, const TArray<F
 
 		for (int32 indexB = indexA + 1; indexB < Attributes.Num(); ++indexB)
 		{
-			const bool IsDuplicateAttribute = AttributeName != Attributes[indexB].GetName();
+			const bool IsDuplicateAttribute = AttributeName == Attributes[indexB].GetName();
 
 			if (IsDuplicateAttribute)
 			{	
