Minor Studio Telemetry API changes
Added FSsessionScope class and STIDIOTELEMETRY_SESSION_SCOPE macro
FStudioTelemetry::StartSession will return true if a session is already running or sucessfully started.
Updated CSVToTelemetry to use the STIDIOTELEMETRY_SESSION_SCOPE macro

diff --git a/Engine/Source/Programs/CSVTools/CSVToTelemetry/Private/CSVToTelemetry.cpp b/Engine/Source/Programs/CSVTools/CSVToTelemetry/Private/CSVToTelemetry.cpp
--- a/Engine/Source/Programs/CSVTools/CSVToTelemetry/Private/CSVToTelemetry.cpp
+++ b/Engine/Source/Programs/CSVTools/CSVToTelemetry/Private/CSVToTelemetry.cpp
@@ -28,7 +28,7 @@ static int GenerateTelemetryFromCSVFile(const FString& FilePath)
 			return Result;
 		}
 
-		FStudioTelemetry::Get().StartSession();
+		STUDIO_TELEMETRY_SESSION_SCOPE
 
 		// Start the telemetry session and record each row as a single telemetry event 
 		if (FStudioTelemetry::Get().IsSessionRunning())
@@ -127,8 +127,6 @@ static int GenerateTelemetryFromCSVFile(const FString& FilePath)
 			{
 				UE_LOG(LogCSVToTelemetry, Error, TEXT("Unable to read rows from CSV file %s"), *FilePath);
 			}
-
-			FStudioTelemetry::Get().EndSession();
 		}
 		else
 		{
diff --git a/Engine/Source/Runtime/StudioTelemetry/Private/StudioTelemetry.cpp b/Engine/Source/Runtime/StudioTelemetry/Private/StudioTelemetry.cpp
--- a/Engine/Source/Runtime/StudioTelemetry/Private/StudioTelemetry.cpp
+++ b/Engine/Source/Runtime/StudioTelemetry/Private/StudioTelemetry.cpp
@@ -51,11 +51,11 @@ void FStudioTelemetry::SetRecordEventCallback(OnRecordEventCallback Callback )
 }
 
 
-void FStudioTelemetry::StartSession()
+bool FStudioTelemetry::StartSession()
 {
 	if (IsSessionRunning())
 	{
-		return;
+		return true;
 	}
 
 	// Load the configuration
@@ -64,7 +64,7 @@ void FStudioTelemetry::StartSession()
 	if (Config.bSendTelemetry == false)
 	{
 		// We did not wish to send any telemetry events
-		return;
+		return false;
 	}
 
 	FScopeLock ScopeLock(&CriticalSection);
@@ -191,7 +191,11 @@ void FStudioTelemetry::StartSession()
 		OnStartSession.Broadcast();
 
 		UE_LOG(LogStudioTelemetry, Log, TEXT("Started StudioTelemetry Session"));
+
+		return true;
 	}
+
+	return false;
 }
 
 void FStudioTelemetry::EndSession()
diff --git a/Engine/Source/Runtime/StudioTelemetry/Public/StudioTelemetry.h b/Engine/Source/Runtime/StudioTelemetry/Public/StudioTelemetry.h
--- a/Engine/Source/Runtime/StudioTelemetry/Public/StudioTelemetry.h
+++ b/Engine/Source/Runtime/StudioTelemetry/Public/StudioTelemetry.h
@@ -42,8 +42,8 @@ public:
 	/** Access to the tracer for the system*/
 	STUDIOTELEMETRY_API TWeakPtr<IAnalyticsTracer> GetTracer();
 
-	/** Starts a new analytics session*/
-	STUDIOTELEMETRY_API void StartSession();
+	/** Starts a new analytics session. Returns true if the session was already running or started successfully*/
+	STUDIOTELEMETRY_API bool StartSession();
 
 	/** Ends an existing analytics session*/
 	STUDIOTELEMETRY_API void EndSession();
@@ -94,11 +94,26 @@ public:
 	DECLARE_MULTICAST_DELEGATE_TwoParams(FOnRecordEvent, const FString&, const TArray<FAnalyticsEventAttribute>&);
 	FOnRecordEvent& GetOnRecordEvent() { return OnRecordEvent; }
 
+	/** Scoped Session helper class **/
+	class FSessionScope
+	{
+	public:
+		FSessionScope()
+		{
+			FStudioTelemetry::Get().StartSession();
+		}
+
+		~FSessionScope()
+		{
+			FStudioTelemetry::Get().EndSession();
+		}
+	};
+
 	/** Scoped Span helper class **/
-	class ScopedSpan
+	class FSpanScope
 	{
 	public:
-		ScopedSpan(const FName Name, const TArray<FAnalyticsEventAttribute>& AdditionalAttributes = {} )
+		FSpanScope(const FName Name, const TArray<FAnalyticsEventAttribute>& AdditionalAttributes = {} )
 		{
 			if (FStudioTelemetry::Get().IsAvailable())
 			{
@@ -106,7 +121,7 @@ public:
 			}
 		}
 
-		~ScopedSpan()
+		~FSpanScope()
 		{
 			if (FStudioTelemetry::Get().IsAvailable())
 			{
@@ -142,7 +157,8 @@ private:
 	FOnRecordEvent							OnRecordEvent;
 };
 
-// Useful macros for scoped spans
-#define STUDIO_TELEMETRY_SPAN_SCOPE(Name) FStudioTelemetry::ScopedSpan PREPROCESSOR_JOIN(ScopedSpan, __LINE__)(TEXT(#Name));
+// Useful macros for scoped sessions and spans
+#define STUDIO_TELEMETRY_SESSION_SCOPE FStudioTelemetry::FSessionScope PREPROCESSOR_JOIN(FSessionScope, __LINE__);
+#define STUDIO_TELEMETRY_SPAN_SCOPE(Name) FStudioTelemetry::FSpanScope PREPROCESSOR_JOIN(FSpanScope, __LINE__)(TEXT(#Name));
 #define STUDIO_TELEMETRY_START_SPAN(Name) if (FStudioTelemetry::Get().IsAvailable()) { FStudioTelemetry::Get().StartSpan(TEXT(#Name));}
 #define STUDIO_TELEMETRY_END_SPAN(Name) if (FStudioTelemetry::Get().IsAvailable()) { FStudioTelemetry::Get().EndSpan(TEXT(#Name));}
