Added an option to CSVToTelemetry NOT to send the SessionContext fields with all events.
This will drastrically reduce the ammount duplicated data uploaded with the event but requires a join with SessionStart event ( which always sends this SessionContext ) via the Session_Id field for equivalent functionality in practice.
By default, SessionContext data is always sent by default to keep existing CSVToTelemetry functionality.
The option can be specifed by the ST_SendSessionContext=true/false on commadline or SendSessionContext=true/false in the program's configuration files.
Added a Session Duration field to the SessionEnd Event.

diff --git a/Engine/Source/Runtime/StudioTelemetry/Private/StudioTelemetry.cpp b/Engine/Source/Runtime/StudioTelemetry/Private/StudioTelemetry.cpp
--- a/Engine/Source/Runtime/StudioTelemetry/Private/StudioTelemetry.cpp
+++ b/Engine/Source/Runtime/StudioTelemetry/Private/StudioTelemetry.cpp
@@ -72,7 +72,8 @@ bool FStudioTelemetry::StartSession()
 
 	if (AnalyticsProvider.IsValid())
 	{
-		TArray<FAnalyticsEventAttribute> DefaultEventAttributes;
+		TArray<FAnalyticsEventAttribute> DefaultAttributes;
+		TArray<FAnalyticsEventAttribute> SessionContextAttributes;
 
 		const FString UserID = FPlatformProcess::UserName(false);
 		const FString ProjectName = FApp::GetProjectName();
@@ -98,45 +99,46 @@ bool FStudioTelemetry::StartSession()
 
 		FGuid SessionID = FApp::GetInstanceId();
 
-		// Set the default event attributes, these will be appended to every event
-		DefaultEventAttributes.Emplace(TEXT("ProjectName"), ProjectName);
-		DefaultEventAttributes.Emplace(TEXT("ProjectID"), ProjectID);
-
-		DefaultEventAttributes.Emplace(TEXT("Session_ID"), SessionID.ToString(EGuidFormats::DigitsWithHyphensInBraces));
-		DefaultEventAttributes.Emplace(TEXT("Session_StartUTC"), FDateTime::UtcNow().ToUnixTimestampDecimal());
+		// Build the default attributes. These are sent with ALL events
+		DefaultAttributes.Emplace(TEXT("Session_ID"), SessionID.ToString(EGuidFormats::DigitsWithHyphensInBraces));
+		DefaultAttributes.Emplace(TEXT("Session_StartUTC"), FDateTime::UtcNow().ToUnixTimestampDecimal());
 
 		FString SessionLabel;
 		if (FParse::Value(FCommandLine::Get(), TEXT("SessionLabel="), SessionLabel))
 		{
-			DefaultEventAttributes.Emplace(TEXT("Session_Label"), SessionLabel);
+			DefaultAttributes.Emplace(TEXT("Session_Label"), SessionLabel);
 		}
 
-		DefaultEventAttributes.Emplace(TEXT("Build_Configuration"), LexToString(FApp::GetBuildConfiguration()));
-		DefaultEventAttributes.Emplace(TEXT("Build_BranchName"), FApp::GetBranchName().ToLower());
-		DefaultEventAttributes.Emplace(TEXT("Build_Changelist"), BuildSettings::GetCurrentChangelist());
+		// Build the session context attributes. These will always be sent with the SessionStart event but can also be appended to default attributes to send with all ALL events.
+		SessionContextAttributes.Emplace(TEXT("ProjectName"), ProjectName);
+		SessionContextAttributes.Emplace(TEXT("ProjectID"), ProjectID);
+
+		SessionContextAttributes.Emplace(TEXT("Build_Configuration"), LexToString(FApp::GetBuildConfiguration()));
+		SessionContextAttributes.Emplace(TEXT("Build_BranchName"), FApp::GetBranchName().ToLower());
+		SessionContextAttributes.Emplace(TEXT("Build_Changelist"), BuildSettings::GetCurrentChangelist());
 
-		DefaultEventAttributes.Emplace(TEXT("Config_IsEditor"), GIsEditor);
-		DefaultEventAttributes.Emplace(TEXT("Config_IsBuildMachine"), GIsBuildMachine);
-		DefaultEventAttributes.Emplace(TEXT("Config_IsRunningCommandlet"), IsRunningCommandlet());
+		SessionContextAttributes.Emplace(TEXT("Config_IsEditor"), GIsEditor);
+		SessionContextAttributes.Emplace(TEXT("Config_IsBuildMachine"), GIsBuildMachine);
+		SessionContextAttributes.Emplace(TEXT("Config_IsRunningCommandlet"), IsRunningCommandlet());
 
 		// Only send user data if requested
 		if (Config.bSendUserData == true)
 		{
-			DefaultEventAttributes.Emplace(TEXT("User_ID"), UserID);
-			DefaultEventAttributes.Emplace(TEXT("Application_Commandline"), FCommandLine::Get());
+			SessionContextAttributes.Emplace(TEXT("User_ID"), UserID);
+			SessionContextAttributes.Emplace(TEXT("Application_Commandline"), FCommandLine::Get());
 		}
 
 		// ALways send the platform
-		DefaultEventAttributes.Emplace(TEXT("Hardware_Platform"), FString(FPlatformProperties::IniPlatformName()));
+		SessionContextAttributes.Emplace(TEXT("Hardware_Platform"), FString(FPlatformProperties::IniPlatformName()));
 
 		// Only send detailed hardware data if requested
 		if (Config.bSendHardwareData == true)
 		{
-			DefaultEventAttributes.Emplace(TEXT("Hardware_CPU"), FPlatformMisc::GetCPUBrand());
-			DefaultEventAttributes.Emplace(TEXT("Hardware_CPU_Cores_Physical"), FPlatformMisc::NumberOfCores());
-			DefaultEventAttributes.Emplace(TEXT("Hardware_CPU_Cores_Logical"), FPlatformMisc::NumberOfCoresIncludingHyperthreads());
-			DefaultEventAttributes.Emplace(TEXT("Hardware_RAM"), static_cast<uint64>(FPlatformMemory::GetStats().TotalPhysical));
-			DefaultEventAttributes.Emplace(TEXT("Hardware_ComputerName"), ComputerName);
+			SessionContextAttributes.Emplace(TEXT("Hardware_CPU"), FPlatformMisc::GetCPUBrand());
+			SessionContextAttributes.Emplace(TEXT("Hardware_CPU_Cores_Physical"), FPlatformMisc::NumberOfCores());
+			SessionContextAttributes.Emplace(TEXT("Hardware_CPU_Cores_Logical"), FPlatformMisc::NumberOfCoresIncludingHyperthreads());
+			SessionContextAttributes.Emplace(TEXT("Hardware_RAM"), static_cast<uint64>(FPlatformMemory::GetStats().TotalPhysical));
+			SessionContextAttributes.Emplace(TEXT("Hardware_ComputerName"), ComputerName);
 		}
 
 		// Only send OS data if requested
@@ -147,10 +149,10 @@ bool FStudioTelemetry::StartSession()
 
 			FPlatformMisc::GetOSVersions(OSVersionLabel, OSSubVersionLabel);
 
-			DefaultEventAttributes.Emplace(TEXT("OS_Version"), FPlatformMisc::GetOSVersion());
-			DefaultEventAttributes.Emplace(TEXT("OS_VersionLabel"), OSVersionLabel);
-			DefaultEventAttributes.Emplace(TEXT("OS_VersionSubLabel"), OSSubVersionLabel);
-			DefaultEventAttributes.Emplace(TEXT("OS_ID"), FPlatformMisc::GetOperatingSystemId());
+			SessionContextAttributes.Emplace(TEXT("OS_Version"), FPlatformMisc::GetOSVersion());
+			SessionContextAttributes.Emplace(TEXT("OS_VersionLabel"), OSVersionLabel);
+			SessionContextAttributes.Emplace(TEXT("OS_VersionSubLabel"), OSSubVersionLabel);
+			SessionContextAttributes.Emplace(TEXT("OS_ID"), FPlatformMisc::GetOperatingSystemId());
 		}
 
 		const FString HordeJobId = GetEnvironmentValriable(TEXT("HordeJobId="), TEXT("UE_HORDE_JOBID")) ;
@@ -158,17 +160,23 @@ bool FStudioTelemetry::StartSession()
 		// Send Horde environment if available
 		if (!HordeJobId.IsEmpty())
 		{
-			DefaultEventAttributes.Emplace(TEXT("Horde_JobID"), HordeJobId);
-			DefaultEventAttributes.Emplace(TEXT("Horde_StepID"), GetEnvironmentValriable(TEXT("HordeStepId="), TEXT("UE_HORDE_STEPID")));
-			DefaultEventAttributes.Emplace(TEXT("Horde_StepName"), GetEnvironmentValriable(TEXT("HordeStepName="), TEXT("UE_HORDE_STEPNAME")));
-			DefaultEventAttributes.Emplace(TEXT("Horde_ServerURL"), GetEnvironmentValriable(TEXT("HordeServerUrl="), TEXT("UE_HORDE_URL")));
-			DefaultEventAttributes.Emplace(TEXT("Horde_TemplateName"), GetEnvironmentValriable(TEXT("HordeTemplateName="), TEXT("UE_HORDE_TEMPLATENAME")));	
+			SessionContextAttributes.Emplace(TEXT("Horde_JobID"), HordeJobId);
+			SessionContextAttributes.Emplace(TEXT("Horde_StepID"), GetEnvironmentValriable(TEXT("HordeStepId="), TEXT("UE_HORDE_STEPID")));
+			SessionContextAttributes.Emplace(TEXT("Horde_StepName"), GetEnvironmentValriable(TEXT("HordeStepName="), TEXT("UE_HORDE_STEPNAME")));
+			SessionContextAttributes.Emplace(TEXT("Horde_ServerURL"), GetEnvironmentValriable(TEXT("HordeServerUrl="), TEXT("UE_HORDE_URL")));
+			SessionContextAttributes.Emplace(TEXT("Horde_TemplateName"), GetEnvironmentValriable(TEXT("HordeTemplateName="), TEXT("UE_HORDE_TEMPLATENAME")));	
+		}
+
+		// Send the Session Context attributes with ALL events by appending to the default attributes if required
+		if (Config.bSendSessionContext == true)
+		{
+			DefaultAttributes.Append(SessionContextAttributes);
 		}
 
 		// Set up the analytics provider
 		AnalyticsProvider->SetUserID(UserID);
 		AnalyticsProvider->SetSessionID(SessionID.ToString(EGuidFormats::DigitsWithHyphensInBraces));
-		AnalyticsProvider->SetDefaultEventAttributes(MoveTemp(DefaultEventAttributes));
+		AnalyticsProvider->SetDefaultEventAttributes(MoveTemp(DefaultAttributes));
 		AnalyticsProvider->SetRecordEventCallback(RecordEventCallback);
 
 		// Start the analytics session
@@ -182,8 +190,17 @@ bool FStudioTelemetry::StartSession()
 		// Bind the pre-exit callback
 		FCoreDelegates::OnEnginePreExit.AddRaw(&FStudioTelemetry::Get(), &FStudioTelemetry::EndSession);
 
+		SessionStartTime = FDateTime::UtcNow();
 		TArray<FAnalyticsEventAttribute> Attributes;
+
+		if (Config.bSendSessionContext == false)
+		{
+			// We always send the session context with the SessionStart event so if we have opted out of sending it with the DefaultAttributes to ALL events we need append it here 
+			Attributes = SessionContextAttributes;
+		}
+
 		Attributes.Emplace(TEXT("SchemaVersion"), 1);
+
 		AnalyticsProvider->RecordEvent(TEXT("StudioTelemetry.SessionStart"), Attributes);
 
 		OnStartSession.Broadcast();
@@ -213,8 +230,11 @@ void FStudioTelemetry::EndSession()
 
 		if (AnalyticsProvider.IsValid())
 		{
-			TArray<FAnalyticsEventAttribute> Attributes;
+			const FTimespan SessionDuration = FDateTime::UtcNow() - SessionStartTime;
+
+			TArray<FAnalyticsEventAttribute> Attributes; 
 			Attributes.Emplace(TEXT("SchemaVersion"), 1);
+			Attributes.Emplace(TEXT("Duration"), SessionDuration.GetTotalSeconds());
 			AnalyticsProvider->RecordEvent(TEXT("StudioTelemetry.SessionEnd"), Attributes);
 
 			AnalyticsProvider->FlushEvents();
@@ -248,6 +268,8 @@ void FStudioTelemetry::LoadConfiguration()
 				GConfig->GetBool(*SectionName, TEXT("SendUserData"), Config.bSendUserData, GEngineIni);
 				GConfig->GetBool(*SectionName, TEXT("SendHardwareData"), Config.bSendHardwareData, GEngineIni);
 				GConfig->GetBool(*SectionName, TEXT("SendOSData"), Config.bSendOSData, GEngineIni);
+				GConfig->GetBool(*SectionName, TEXT("SendSessionContext"), Config.bSendSessionContext, GEngineIni);
+				
 			}
 		}
 	}
@@ -257,6 +279,9 @@ void FStudioTelemetry::LoadConfiguration()
 	FParse::Bool(FCommandLine::Get(), TEXT("ST_SendUserData="), Config.bSendUserData);
 	FParse::Bool(FCommandLine::Get(), TEXT("ST_SendHardwareData="), Config.bSendHardwareData);
 	FParse::Bool(FCommandLine::Get(), TEXT("ST_SendOSData="), Config.bSendOSData);
+	FParse::Bool(FCommandLine::Get(), TEXT("ST_SendSessionContext="), Config.bSendSessionContext);
+
+	
 }
 
 void FStudioTelemetry::RecordEvent(const FString& EventName, const TArray<FAnalyticsEventAttribute>& Attributes)
diff --git a/Engine/Source/Runtime/StudioTelemetry/Public/StudioTelemetry.h b/Engine/Source/Runtime/StudioTelemetry/Public/StudioTelemetry.h
--- a/Engine/Source/Runtime/StudioTelemetry/Public/StudioTelemetry.h
+++ b/Engine/Source/Runtime/StudioTelemetry/Public/StudioTelemetry.h
@@ -144,6 +144,7 @@ private:
 		bool bSendUserData = false;  // Never send user data unless specifically asked to
 		bool bSendHardwareData = false; // Never send hardware data unless specifically asked to4
 		bool bSendOSData = false; // Never send operating system data unless specifically asked to
+		bool bSendSessionContext = true; // Send the session context data with every event
 	};
 	
 	FCriticalSection						CriticalSection;
@@ -152,6 +153,7 @@ private:
 	OnRecordEventCallback					RecordEventCallback;
 	FGuid									SessionGUID;
 	FConfig									Config;
+	FDateTime								SessionStartTime;
 	FOnStartSession							OnStartSession;
 	FOnEndSession							OnEndSession;
 	FOnRecordEvent							OnRecordEvent;
