Added IsDebuggerPresent to StudioTelemetry Session Context
Bumped SessionContextSchema to 2

diff --git a/Engine/Source/Runtime/StudioTelemetry/Private/StudioTelemetry.cpp b/Engine/Source/Runtime/StudioTelemetry/Private/StudioTelemetry.cpp
--- a/Engine/Source/Runtime/StudioTelemetry/Private/StudioTelemetry.cpp
+++ b/Engine/Source/Runtime/StudioTelemetry/Private/StudioTelemetry.cpp
@@ -120,6 +120,7 @@ bool FStudioTelemetry::StartSession()
 		SessionContextAttributes.Emplace(TEXT("Config_IsEditor"), GIsEditor);
 		SessionContextAttributes.Emplace(TEXT("Config_IsBuildMachine"), GIsBuildMachine);
 		SessionContextAttributes.Emplace(TEXT("Config_IsRunningCommandlet"), IsRunningCommandlet());
+		SessionContextAttributes.Emplace(TEXT("Config_IsDebuggerPresent"), FPlatformMisc::IsDebuggerPresent());
 
 		// Only send user data if requested
 		if (Config.bSendUserData == true)
@@ -199,7 +200,9 @@ bool FStudioTelemetry::StartSession()
 			Attributes = SessionContextAttributes;
 		}
 
-		Attributes.Emplace(TEXT("SchemaVersion"), 1);
+		const int32 SessionContextSchemaVersion = 2;
+
+		Attributes.Emplace(TEXT("SchemaVersion"), SessionContextSchemaVersion);
 
 		AnalyticsProvider->RecordEvent(TEXT("StudioTelemetry.SessionStart"), Attributes);
 
